Scripts that are useful for using UnrealAssetStringify to diff from p4v or open .uassets from explorer

#rnx

diff --git a/Engine/Source/Programs/UnrealAssetStringify/BuildScripts/UnrealAssetStringifyInstall.bat b/Engine/Source/Programs/UnrealAssetStringify/BuildScripts/UnrealAssetStringifyInstall.bat
new file mode 100755
--- /dev/null
+++ b/Engine/Source/Programs/UnrealAssetStringify/BuildScripts/UnrealAssetStringifyInstall.bat
@@ -0,0 +1,20 @@
+@echo off
+REM Copyright Epic Games, Inc. All Rights Reserved.
+REM -----------------------------------------------
+REM UnrealAssetStringifyInstall.bat
+REM Batch file that runs UnrealAssetStringifyInstall.py
+REM -----------------------------------------------
+
+setlocal
+set SCRIPT_DIR=%~dp0
+
+REM Normalize path (remove trailing backslash)
+if "%SCRIPT_DIR:~-1%"=="\" set SCRIPT_DIR=%SCRIPT_DIR:~0,-1%
+
+set PYTHON_EXE=%SCRIPT_DIR%\..\..\..\..\Binaries\ThirdParty\Python3\Win64\python.exe
+set SCRIPT_PATH=%SCRIPT_DIR%\UnrealAssetStringifyInstall.py
+
+"%PYTHON_EXE%" "%SCRIPT_PATH%"
+pause
+
+endlocal
\ No newline at end of file
diff --git a/Engine/Source/Programs/UnrealAssetStringify/BuildScripts/UnrealAssetStringifyInstall.py b/Engine/Source/Programs/UnrealAssetStringify/BuildScripts/UnrealAssetStringifyInstall.py
new file mode 100644
--- /dev/null
+++ b/Engine/Source/Programs/UnrealAssetStringify/BuildScripts/UnrealAssetStringifyInstall.py
@@ -0,0 +1,123 @@
+# Copyright Epic Games, Inc. All Rights Reserved.
+#
+# Script to build UnrealAssetStringify via RunUBT and setup a umap/
+# uasset diff assocation in P4V. This script is windows only but could
+# be made to work on linux and os x with some effort and testing
+
+import os
+import subprocess
+import xml.etree.ElementTree as ET
+import shutil
+from pathlib import Path
+from datetime import datetime
+
+SCRIPT_DIR = Path(__file__).parent
+ENGINE_DIR = SCRIPT_DIR.parent.parent.parent.parent
+
+def get_unreal_stringify_path() -> Path:
+    # helper for .exe/noextension. I want to query existence so 
+    # haven't found away around this platform check:
+    exe_name = "UnrealAssetStringify"
+    print(os.name)
+    if os.name == "nt":
+        # On Windows, prefer .exe if it exists
+        path = ENGINE_DIR / "Binaries" / "Win64" / f"{exe_name}.exe"
+        if path.exists():
+            return path
+    # Otherwise, fallback to no extension
+    return ENGINE_DIR / "Binaries" / exe_name
+
+
+def build_unreal_asset_stringify():
+    ubt_bat = ENGINE_DIR / "Build" / "BatchFiles" / "RunUBT.bat"
+    cmd = [str(ubt_bat), "UnrealAssetStringify", "Win64", "Development"]
+    print(f"Running: {' '.join(cmd)}")
+    subprocess.run(cmd, check=True)
+    if not get_unreal_stringify_path().exists():
+        print("🚫 UnrealAssetStringify.exe does not exist, despite RunUBT returning 0")
+        sys.exit(1)
+    print("✅ Compiled UnrealAssetStringify.exe successfully")
+
+
+def update_p4v_diff_associations():
+    """Insert UnrealAssetStringifyP4VDiff associations into P4V ApplicationSettings.xml."""
+    user_home = Path.home()
+    xml_path = user_home / ".p4qt" / "ApplicationSettings.xml"
+    python_path = ENGINE_DIR / "Binaries" / "ThirdParty" / "Python3" / "Win64" / "python.exe"
+    diff_script = ENGINE_DIR / "Source" / "Programs" / "UnrealAssetStringify" / "Scripts" / "UnrealAssetStringifyP4VDiff.py"
+
+    if not xml_path.exists():
+        print(f"🚫 {xml_path} not found — skipping P4V diff integration.")
+        return
+    
+    # Create backup before modifying
+    timestamp = datetime.now().strftime("%Y-%m-%d_%H%M%S")
+    backup_path = xml_path.with_name(f"ApplicationSettings_{timestamp}.bak.xml")
+    try:
+        shutil.copy2(xml_path, backup_path)
+        print(f"💾 Backup created at: {backup_path}")
+    except Exception as e:
+        print(f"⚠️ Failed to create backup: {e} - associations not updated")
+        return
+
+    print(f"Updating P4V ApplicationSettings.xml at: {xml_path}")
+    tree = ET.parse(xml_path)
+    root = tree.getroot()
+
+    # Find the ApplicationSettings PropertyList
+    if root.tag == "PropertyList" and root.get("varName") == "ApplicationSettings":
+        app_settings = root
+    else:
+        app_settings = root.find(".//PropertyList[@varName='ApplicationSettings']")
+    if app_settings is None:
+        print("🚫 Could not locate <PropertyList varName='ApplicationSettings'> in XML.")
+        return
+
+    # Find (or create) the DiffAssociations node
+    diff_assocs = app_settings.find(".//Associations[@varName='DiffAssociations']")
+    if diff_assocs is None:
+        diff_assocs = ET.SubElement(app_settings, "Associations", {"varName": "DiffAssociations"})
+        prop_list = ET.SubElement(diff_assocs, "PropertyList", {"IsManaged": "TRUE", "varName": "Associations"})
+    else:
+        prop_list = diff_assocs.find(".//PropertyList[@varName='Associations']")
+        if prop_list is None:
+            prop_list = ET.SubElement(diff_assocs, "PropertyList", {"IsManaged": "TRUE", "varName": "Associations"})
+
+    # Remove old uasset/umap associations, if any:
+    for assoc in list(prop_list.findall(".//Association")):
+        var_name = assoc.attrib.get("varName")
+        if var_name in ("uasset", "umap"):
+            prop_list.remove(assoc)
+
+    # Add uasset and umap associations w/ UnrealAssetStringifyP4VDiff
+    for var in ("uasset", "umap"):
+        new_assoc = ET.SubElement(prop_list, "Association", {"varName": var})
+        ET.SubElement(new_assoc, "Application").text = str(python_path)
+        ET.SubElement(new_assoc, "Arguments").text = f"{diff_script} %1 %2"
+
+    # Write back the updated XML
+    tree.write(xml_path, encoding="utf-8", xml_declaration=True)
+    print("✅ Updated uasset/umap diff associations in P4V")
+
+def is_p4v_running():
+    """Check if P4V is currently running, it writes ApplicationSettings.xml on close and will discard changes if we don't close it"""
+    # this only works on windows, would have to use ps or similar wo/ psutil
+    result = subprocess.run("tasklist", capture_output=True, text=True)
+    return "p4v.exe" in result.stdout.lower()
+
+def install_unreal_asset_stringify():
+    build_unreal_asset_stringify() # builds our helpful tool
+    
+    # set up p4 file assocations, if we can:
+    if is_p4v_running():
+        print("⚠️ You must close P4V if you want to update P4V diff tool assocations as closing P4V will overwrite settings and they do not hot load")
+    add_p4v = input("❔ Add UnrealAssetStringifyP4VDiff to P4V diff tool config? (y/n): ").strip().lower()
+    if add_p4v == "y":
+        if is_p4v_running():
+            print("⚠️ P4V assocations not updated - close P4V and rerun this script if you want them")
+            return
+        update_p4v_diff_associations()
+
+
+if __name__ == "__main__":
+    install_unreal_asset_stringify()
\ No newline at end of file
diff --git a/Engine/Source/Programs/UnrealAssetStringify/Scripts/UnrealAssetStringifyOpenAsset.bat b/Engine/Source/Programs/UnrealAssetStringify/Scripts/UnrealAssetStringifyOpenAsset.bat
new file mode 100755
--- /dev/null
+++ b/Engine/Source/Programs/UnrealAssetStringify/Scripts/UnrealAssetStringifyOpenAsset.bat
@@ -0,0 +1,23 @@
+@echo off
+REM Copyright Epic Games, Inc. All Rights Reserved.
+REM -----------------------------------------------
+REM UnrealAssetStringifyOpenAsset.bat
+REM Batch file so that UnrealAssetStringifyOpenAsset.py can be run 
+REM via explorer (e.g. Open With). If you need a change in behavior
+REM here consider making it in the .py. Batch is difficult to maintain
+REM -----------------------------------------------
+
+setlocal
+set SCRIPT_DIR=%~dp0
+
+REM Normalize path (remove trailing backslash)
+if "%SCRIPT_DIR:~-1%"=="\" set SCRIPT_DIR=%SCRIPT_DIR:~0,-1%
+
+REM relative addressing of python and UnrealAssetStringifyOpenAsset.py:
+set PYTHON_EXE=%SCRIPT_DIR%\..\..\..\..\Binaries\ThirdParty\Python3\Win64\python.exe
+set SCRIPT_PATH=%SCRIPT_DIR%\UnrealAssetStringifyOpenAsset.py
+
+REM Forward all arguments to Python
+start "" /min "%PYTHON_EXE%" "%SCRIPT_PATH%" %*
+
+endlocal
\ No newline at end of file
diff --git a/Engine/Source/Programs/UnrealAssetStringify/Scripts/UnrealAssetStringifyOpenAsset.py b/Engine/Source/Programs/UnrealAssetStringify/Scripts/UnrealAssetStringifyOpenAsset.py
new file mode 100644
--- /dev/null
+++ b/Engine/Source/Programs/UnrealAssetStringify/Scripts/UnrealAssetStringifyOpenAsset.py
@@ -0,0 +1,53 @@
+# Copyright Epic Games, Inc. All Rights Reserved.
+
+import subprocess
+import tempfile
+import os
+import sys
+import shutil
+from pathlib import Path
+
+SCRIPT_DIR = Path(__file__).parent
+UNREAL_BIN_DIR = SCRIPT_DIR.parent.parent.parent.parent / "Binaries" / "Win64"
+UNREAL_STRINGIFY = SCRIPT_DIR / UNREAL_BIN_DIR / "UnrealAssetStringify"
+
+EDITOR = "code"
+
+def run_unreal_asset_stringify(asset_path):
+    cmd = [str(UNREAL_STRINGIFY), asset_path]
+    print(f"Running: {' '.join(cmd)}")
+
+    try:
+        result = subprocess.run(cmd, capture_output=True, text=True)
+        json_output = result.stdout.strip()
+    except FileNotFoundError:
+        json_output = f"Error: UnrealAssetStringify not found at {asset_path} - do you want to remove your file association?"
+
+    if not json_output:
+        json_output = "Error: No output received from UnrealAssetStringify."
+
+    temp_file = tempfile.NamedTemporaryFile(delete=False, suffix=".json")
+    temp_file.write(json_output.encode("utf-8"))
+    temp_file.close()
+
+    print(f"JSON written to: {temp_file.name}")
+    return temp_file.name
+
+
+def open_in_vscode(file_path):
+    vscode_cmd = shutil.which(EDITOR)
+    if not vscode_cmd:
+        print(f"Error: {EDITOR} command not found in PATH.")
+        sys.exit(1)
+
+    subprocess.run([vscode_cmd, file_path])
+
+
+if __name__ == "__main__":
+    if len(sys.argv) < 2:
+        print("Usage: python UnrealAssetStringifyOpenAsset.py <AssetPath>")
+        sys.exit(1)
+
+    asset_path = sys.argv[1]
+    json_file = run_unreal_asset_stringify(asset_path)
+    open_in_vscode(json_file)
\ No newline at end of file
diff --git a/Engine/Source/Programs/UnrealAssetStringify/Scripts/UnrealAssetStringifyP4VDiff.py b/Engine/Source/Programs/UnrealAssetStringify/Scripts/UnrealAssetStringifyP4VDiff.py
new file mode 100644
--- /dev/null
+++ b/Engine/Source/Programs/UnrealAssetStringify/Scripts/UnrealAssetStringifyP4VDiff.py
@@ -0,0 +1,88 @@
+# Copyright Epic Games, Inc. All Rights Reserved.
+#
+# This script prints .uasset files to stdout and displays them
+# in the p4merge diff program, if available. It is meant to be
+# integrated as custom diff tool in p4v
+#
+# To enable this tool edit [REDACTED_USER_HOME]\.p4qt\ApplicationSettings.xml:
+r'''
+<Associations varName="DiffAssociations">
+ <RunExternal>false</RunExternal>
+ <PropertyList IsManaged="TRUE" varName="Associations">
+  <Association varName="uasset">
+   <Application>D:\path\to\Engine\Binaries\ThirdParty\Python3\Win64\python.exe</Application>
+   <Arguments>D:\path\to\Engine\Source\Programs\UnrealAssetStringify\Scripts\UnrealAssetStringifyP4VDiff.py %1 %2</Arguments>
+  </Association>
+  <Association varName="umap">
+   <Application>D:\path\to\Engine\Binaries\ThirdParty\Python3\Win64\python.exe</Application>
+   <Arguments>D:\path\to\Engine\Source\Programs\UnrealAssetStringify\Scripts\UnrealAssetStringifyP4VDiff.py %1 %2</Arguments>
+  </Association>
+ </PropertyList>
+</Associations>
+
+in <PropertyList IsManaged="TRUE" varName="ApplicationSettings"></PropertyList>
+'''
+# There is a windows only install script in ..BuildScripts/UnrealAssetStringifyInstall.py
+
+import os
+import sys
+import tempfile
+import subprocess
+from pathlib import Path
+
+SCRIPT_DIR = Path(__file__).parent
+UNREAL_BIN_DIR = SCRIPT_DIR.parent.parent.parent.parent / "Binaries" / "Win64"
+
+def get_unreal_stringify_path() -> Path:
+    # helper for .exe/noextension. I want to query existence so 
+    # haven't found away around this platform check:
+    exe_name = "UnrealAssetStringify"
+    if os.name == "nt":
+        # On Windows, prefer .exe if it exists
+        path = UNREAL_BIN_DIR / f"{exe_name}.exe"
+        if path.exists():
+            return path
+    # Otherwise, fallback to no extension
+    return UNREAL_BIN_DIR / exe_name
+
+UNREAL_STRINGIFY = get_unreal_stringify_path()
+P4MERGE = "p4merge"
+
+def write_src_as_json_to_dst(exe: Path, src: Path, dst: Path) -> None:
+    with open(dst, "wb") as out:
+        proc = subprocess.run([str(exe), str(src)], stdout=out, stderr=out)
+
+def main(argv: list[str]) -> int:
+    if len(argv) < 3:
+        print("Usage: python UnrealAssetStringifyP4VDiff.py <left> <right>", file=sys.stderr)
+        return 1
+
+    left_in, right_in = map(Path, argv[1:3])
+    tmp_dir = Path(tempfile.mkdtemp(prefix="p4vjsondiff_"))
+    ljson, rjson = tmp_dir / "left.json", tmp_dir / "right.json"
+
+    exe_exists = UNREAL_STRINGIFY.exists()
+    exe_missing_not_exists = (
+        "UnrealAssetStringify not found.\n"
+        "Please build it by running:\n"
+        ">RunUBT.bat UnrealAssetStringify Win64 Development\n"
+        "From your Engine root directory"
+    )
+    error_msg_left = ""
+    error_msg_right = ""
+    if exe_exists:
+        write_src_as_json_to_dst(UNREAL_STRINGIFY, left_in, ljson)
+        write_src_as_json_to_dst(UNREAL_STRINGIFY, right_in, rjson)
+    else:
+        error_msg_left = exe_missing_not_exists
+        error_msg_right = exe_missing_not_exists
+
+    if error_msg_left != "":
+        ljson.write_text(error_msg_left, encoding="utf-8")
+    if error_msg_right != "":
+        rjson.write_text(error_msg_right, encoding="utf-8")
+
+    return subprocess.call([P4MERGE, str(ljson), str(rjson)])
+
+if __name__ == "__main__":
+    sys.exit(main(sys.argv))
\ No newline at end of file
