Tweak UnrealAssetStringifyInstall to work around zombie p4v processes preventing diff tool association, reduce some verbose messages in the install script

#rnx

diff --git a/Engine/Source/Programs/UnrealAssetStringify/BuildScripts/UnrealAssetStringifyInstall.py b/Engine/Source/Programs/UnrealAssetStringify/BuildScripts/UnrealAssetStringifyInstall.py
--- a/Engine/Source/Programs/UnrealAssetStringify/BuildScripts/UnrealAssetStringifyInstall.py
+++ b/Engine/Source/Programs/UnrealAssetStringify/BuildScripts/UnrealAssetStringifyInstall.py
@@ -8,6 +8,7 @@ import os
 import subprocess
 import xml.etree.ElementTree as ET
 import shutil
+from time import sleep
 from pathlib import Path
 from datetime import datetime
 
@@ -55,7 +56,7 @@ def update_p4v_diff_associations():
     backup_path = xml_path.with_name(f"ApplicationSettings_{timestamp}.bak.xml")
     try:
         shutil.copy2(xml_path, backup_path)
-        print(f"💾 Backup created at: {backup_path}")
+        print(f"💾 P4V ApplicationSettings.xml backup created at: {backup_path}")
     except Exception as e:
         print(f"⚠️ Failed to create backup: {e} - associations not updated")
         return
@@ -100,21 +101,59 @@ def update_p4v_diff_associations():
     print("✅ Updated uasset/umap diff associations in P4V")
 
 def is_p4v_running():
-    """Check if P4V is currently running, it writes ApplicationSettings.xml on close and will discard changes if we don't close it"""
     # this only works on windows, would have to use ps or similar wo/ psutil
     result = subprocess.run("tasklist", capture_output=True, text=True)
     return "p4v.exe" in result.stdout.lower()
 
+def close_p4v():
+    # close any running p4v processes, so that we can update ApplicationSettings.xml
+    # and not have our change overwritten when those running processes finally exit:
+    print("Closing running P4V instances...")
+    try:
+        subprocess.run(["taskkill", "/IM", "p4v.exe"], check=True, capture_output=True)
+    except Exception as e:
+        print("⚠️ Unexpected error invoking taskkill (graceful):", e)
+        return False
+
+    for _ in range(10):
+        if not is_p4v_running():
+            return True
+        print("Waiting for P4V to close...")
+        sleep(1.0)
+    
+    print("⚠️ P4V Graceful exit failed, forcing P4V to close")
+    try:
+        subprocess.run(["taskkill", "/F", "/IM", "p4v.exe"], check=False, capture_output=True)
+    except Exception as e:
+        print("⚠️ Failed to close P4V processes.")
+        print(e.stderr.decode() if e.stderr else str(e))
+        return False
+        
+    for _ in range(3):
+        if not is_p4v_running():
+            return True
+        print("Waiting for P4V to close...")
+        sleep(1.0)
+
+    return True
+
 def install_unreal_asset_stringify():
+    # first, ask the user what they want to do:
+    print("This script will build UnrealAssetStringify and (optionally) associate with p4v uasset/umap diff commmands")
+    if is_p4v_running():
+        print("⚠️ P4V is running - adding UnrealAssetStringifyP4VDiff to config will close P4V")
+    add_p4v = input("❔ Add UnrealAssetStringifyP4VDiff to P4V diff tool config? (y/n): ").strip().lower()
+
+    # no more prompts for input after this point, go ahead and run the installer:
     build_unreal_asset_stringify() # builds our helpful tool
     
     # set up p4 file assocations, if we can:
-    if is_p4v_running():
-        print("⚠️ You must close P4V if you want to update P4V diff tool assocations as closing P4V will overwrite settings and they do not hot load")
-    add_p4v = input("❔ Add UnrealAssetStringifyP4VDiff to P4V diff tool config? (y/n): ").strip().lower()
     if add_p4v == "y":
         if is_p4v_running():
-            print("⚠️ P4V assocations not updated - close P4V and rerun this script if you want them")
+            if not close_p4v():
+                return
+        if is_p4v_running():
+            print("⚠️ P4V assocations not updated - P4V could not be closed")
             return
         update_p4v_diff_associations()
 
