set IDO.Enable = false for 40.00 with the exclusion of the stringify tool
#Synced-CL 50922333

diff --git a/Engine/Source/Programs/UnrealAssetStringify/Private/UnrealAssetStringify.cpp b/Engine/Source/Programs/UnrealAssetStringify/Private/UnrealAssetStringify.cpp
--- a/Engine/Source/Programs/UnrealAssetStringify/Private/UnrealAssetStringify.cpp
+++ b/Engine/Source/Programs/UnrealAssetStringify/Private/UnrealAssetStringify.cpp
@@ -284,6 +284,9 @@ INT32_MAIN_INT32_ARGC_TCHAR_ARGV()
 	FLogSuppressionInterface::Get().ProcessConfigAndCommandLine();
 	// disable localization display string support, it is expensive to initialize:
 	IConsoleManager::Get().ProcessUserConsoleInput(TEXT("Localization.DisplayStringSupport 2"), *GLog, nullptr);
+	// ensure all needed ido and placeholder features are enabled
+	IConsoleManager::Get().ProcessUserConsoleInput(TEXT("IDO.Enable 1"), *GLog, nullptr);
+	IConsoleManager::Get().ProcessUserConsoleInput(TEXT("IDO.Placeholder.Enable 1"), *GLog, nullptr);
 	IConsoleManager::Get().ProcessUserConsoleInput(TEXT("IDO.Placeholder.Feature.ReplaceMissingTypeImportsOnLoad 1"), *GLog, nullptr);
 	IConsoleManager::Get().ProcessUserConsoleInput(TEXT("IDO.Placeholder.Feature.SerializeExportReferencesOnLoad 1"), *GLog, nullptr);
 	IConsoleManager::Get().ProcessUserConsoleInput(TEXT("IDO.Placeholder.Feature.ReplaceMissingReinstancedTypes 1"), *GLog, nullptr);
diff --git a/Engine/Source/Runtime/CoreUObject/Private/UObject/InstanceDataObjectUtils.cpp b/Engine/Source/Runtime/CoreUObject/Private/UObject/InstanceDataObjectUtils.cpp
--- a/Engine/Source/Runtime/CoreUObject/Private/UObject/InstanceDataObjectUtils.cpp
+++ b/Engine/Source/Runtime/CoreUObject/Private/UObject/InstanceDataObjectUtils.cpp
@@ -305,7 +305,7 @@ namespace UE
 
 	bool bEverCreatedIDO = false;
 
-	bool bEnableIDOSupport = true;
+	bool bEnableIDOSupport = false;
 	bool bEverEnabledIDOSupport = true;
 	FAutoConsoleVariableRef EnableIDOSupportCVar(
 		TEXT("IDO.Enable"),
