Fix bad dash

#rnx
#Synced-CL 52588543

diff --git a/Engine/Source/Programs/UnrealAssetStringify/BuildScripts/UnrealAssetStringifyInstall.py b/Engine/Source/Programs/UnrealAssetStringify/BuildScripts/UnrealAssetStringifyInstall.py
--- a/Engine/Source/Programs/UnrealAssetStringify/BuildScripts/UnrealAssetStringifyInstall.py
+++ b/Engine/Source/Programs/UnrealAssetStringify/BuildScripts/UnrealAssetStringifyInstall.py
@@ -48,7 +48,7 @@ def update_p4v_diff_associations():
     diff_script = ENGINE_DIR / "Source" / "Programs" / "UnrealAssetStringify" / "Scripts" / "UnrealAssetStringifyP4VDiff.py"
 
     if not xml_path.exists():
-        print(f"🚫 {xml_path} not found — skipping P4V diff integration.")
+        print(f"🚫 {xml_path} not found - skipping P4V diff integration.")
         return
     
     # Create backup before modifying
