UnrealAssetStringify improvements: Add optional explicit save file location, with safe defaults (no overwrite unless --force), add abbreviated arguments (-s, -f)

#rnx
#Synced-CL 52669209

diff --git a/Engine/Source/Programs/UnrealAssetStringify/Private/UnrealAssetStringify.cpp b/Engine/Source/Programs/UnrealAssetStringify/Private/UnrealAssetStringify.cpp
--- a/Engine/Source/Programs/UnrealAssetStringify/Private/UnrealAssetStringify.cpp
+++ b/Engine/Source/Programs/UnrealAssetStringify/Private/UnrealAssetStringify.cpp
@@ -17,12 +17,17 @@ const TCHAR* const Usage = TEXT(R"(Usage:
 	UnrealAssetStringify [OPTION]... FILE
 
 Options:
-	--save	Save the stringified output to Engine\Programs\UnrealAssetStringify\Saved
-			instead of simply printing it
+	--save, -s         Save the stringified output to Engine\Programs\UnrealAssetStringify\Saved
+	                   instead of simply printing it
+	--save=PATH, -s=PATH
+	                   Save the stringified output to an explicit file path
+	--force, -f        Overwrite the output file if it already exists
 
 Example:
 	UnrealAssetStringify.exe "D:\Path\To\A.uasset"
-	UnrealAssetStringify.exe --save "D:\Path\To\A.uasset")");
+	UnrealAssetStringify.exe --save "D:\Path\To\A.uasset"
+	UnrealAssetStringify.exe --save="D:\Output\A.json" "D:\Path\To\A.uasset"
+	UnrealAssetStringify.exe -s="D:\Output\A.json" -f "D:\Path\To\A.uasset")");
 	UE_LOGFMT(LogUnrealAssetStringify, Display, "{0}", Usage);
 }
 
@@ -37,6 +42,7 @@ namespace UE::Private
 		MissingFile,
 		BadAsset,
 		BadPackageName,
+		FileAlreadyExists,
 		CouldNotCreateResultFile,
 		CouldNotWriteResultFile
 	};
@@ -100,25 +106,56 @@ struct FUnrealAssetStringifyArgs
 {
 	// options:
 	bool bSave = false;
+	FString SavePath;   // explicit output path; empty means use default save dir
+	bool bForce = false;
 	// required inputs:
 	FString Filename;
-	
+
 	static bool ParseArgs(int32 ArgC, TCHAR* const ArgV[], FUnrealAssetStringifyArgs& OutArgs);
 	static int32 StringifyFile(FStringView Filename, const FUnrealAssetStringifyArgs& InArgs);
 };
 
 bool FUnrealAssetStringifyArgs::ParseArgs(int32 ArgC, TCHAR* const ArgV[], FUnrealAssetStringifyArgs& OutArgs)
 {
+	// Helper: check if Arg matches a flag name with an optional =VALUE suffix.
+	// Returns true when the flag matches; if a value is present it is placed in OutValue.
+	auto TryParseFlag = [](FStringView Arg, FStringView Flag, FString& OutValue) -> bool
+	{
+		if(!Arg.StartsWith(Flag))
+		{
+			return false;
+		}
+		const FStringView Remainder = Arg.Mid(Flag.Len());
+		if(Remainder.IsEmpty())
+		{
+			OutValue = TEXT("");
+			return true;
+		}
+		if(Remainder[0] == TEXT('='))
+		{
+			OutValue = FString(Remainder.Mid(1));
+			return true;
+		}
+		return false; // e.g. "--saves" should not match "--save"
+	};
+
 	int32 I;
 	for(I = 1; I < ArgC - 1; ++I)
 	{
-		if(FStringView(ArgV[I]) == TEXT("--save"))
+		const FStringView Arg(ArgV[I]);
+		FString Value;
+		if(TryParseFlag(Arg, TEXT("--save"), Value) || TryParseFlag(Arg, TEXT("-s"), Value))
 		{
 			OutArgs.bSave = true;
+			OutArgs.SavePath = MoveTemp(Value);
+		}
+		else if(Arg == TEXT("--force") || Arg == TEXT("-f"))
+		{
+			OutArgs.bForce = true;
 		}
 	}
 
-	if(I < ArgC && !FStringView(ArgV[I]).StartsWith(TEXT("--")))
+	if(I < ArgC && !FStringView(ArgV[I]).StartsWith(TEXT("-")))
 	{
 		OutArgs.Filename = ArgV[I];
 	}
@@ -191,32 +228,46 @@ int32 FUnrealAssetStringifyArgs::StringifyFile(FStringView Filename, const FUnre
 	}
 	else
 	{
-		// filename to use, _uas is just a fingerprint, short for
-		// UnrealAssetStringify, while .json is a reasonable file extension:
-		const auto GetIntermediateAssetName = [](const UObject* Object)
+		FString DestinationFile;
+		if(!InArgs.SavePath.IsEmpty())
+		{
+			// Explicit path supplied via --save=PATH or -s=PATH
+			DestinationFile = InArgs.SavePath;
+			FPaths::NormalizeFilename(DestinationFile);
+		}
+		else
 		{
-			return
+			// Default path: Saved/<package-path>_uas.json
+			// _uas is a fingerprint for UnrealAssetStringify, .json is the extension
+			DestinationFile =
 				FPaths::ProjectSavedDir() +
-				Object->GetPathName() +
+				Package->GetPathName() +
 				FString(TEXT("_uas.json"));
-		};
-		FString DestinationFile = GetIntermediateAssetName(Package);
+		}
+
+		if(IFileManager::Get().FileExists(*DestinationFile) && !InArgs.bForce)
+		{
+			UE_LOGFMT(LogUnrealAssetStringify, Error,
+				"Output file already exists: {0} (use --force / -f to overwrite)", DestinationFile);
+			return (int32)UE::Private::EReturnCodes::FileAlreadyExists;
+		}
+
 		TUniquePtr<FArchive> FileArchive(IFileManager::Get().CreateDebugFileWriter(*DestinationFile));
 		if(!FileArchive) // failed to create a file
 		{
-			UE_LOGFMT(LogUnrealAssetStringify, Error, 
+			UE_LOGFMT(LogUnrealAssetStringify, Error,
 				"Failed to create result file {0}", DestinationFile);
-			return (int32)UE::Private::EReturnCodes::CouldNotCreateResultFile; 
+			return (int32)UE::Private::EReturnCodes::CouldNotCreateResultFile;
 		}
 
 		FileArchive->Serialize((void*)Json.GetCharArray().GetData(), Json.GetCharArray().Num() - 1);
 		if(FileArchive->IsError()) // error during writing, likely disk space
 		{
-			UE_LOGFMT(LogUnrealAssetStringify, Error, 
+			UE_LOGFMT(LogUnrealAssetStringify, Error,
 				"Failed to write result file {0}", DestinationFile);
-			return (int32)UE::Private::EReturnCodes::CouldNotWriteResultFile; 
+			return (int32)UE::Private::EReturnCodes::CouldNotWriteResultFile;
 		}
-		UE_LOGFMT(LogUnrealAssetStringify, Display, 
+		UE_LOGFMT(LogUnrealAssetStringify, Display,
 			"Wrote result file: {0}", DestinationFile);
 	}
 	return (int32)UE::Private::EReturnCodes::Success;
